Stopped freeing memory that was not my responsibility to free (which was causing a crash when trying multiple times to select fonts that don't fit Offline Font criteria)

diff --git a/Engine/Source/Developer/SlateFontDialog/Private/SlateFontDlgWindow.cpp b/Engine/Source/Developer/SlateFontDialog/Private/SlateFontDlgWindow.cpp
--- a/Engine/Source/Developer/SlateFontDialog/Private/SlateFontDlgWindow.cpp
+++ b/Engine/Source/Developer/SlateFontDialog/Private/SlateFontDlgWindow.cpp
@@ -165,11 +165,6 @@ FSlateFontDlgWindow::FSlateFontDlgWindow(bool& OutSuccess)
     .SetColorAndOpacity(FontColor);
 }
 
-FSlateFontDlgWindow::~FSlateFontDlgWindow()
-{
-    FcFontSetDestroy(FontSet);
-}
-
 BEGIN_SLATE_FUNCTION_BUILD_OPTIMIZATION
 
 // OutFontName actually gets set to the filepath of the chosen font to reduce dependencies in UTrueTypeFontFactory::LoadFontFace within TTFontImport.cpp 
diff --git a/Engine/Source/Developer/SlateFontDialog/Private/SlateFontDlgWindow.h b/Engine/Source/Developer/SlateFontDialog/Private/SlateFontDlgWindow.h
--- a/Engine/Source/Developer/SlateFontDialog/Private/SlateFontDlgWindow.h
+++ b/Engine/Source/Developer/SlateFontDialog/Private/SlateFontDlgWindow.h
@@ -45,8 +45,6 @@ private:
 public:
     FSlateFontDlgWindow(bool& OutSuccess);
 
-    ~FSlateFontDlgWindow();
-
     void OpenFontWindow(FString& OutFontName, float& OutHeight, EFontImportFlags& OutFlags, bool& OutSuccess);
 };
 
