Updating SlateFontDialog to use FAppStyle instead of FEditorStyle

diff --git a/Engine/Source/Developer/SlateFontDialog/Private/SlateFontDlgWindow.cpp b/Engine/Source/Developer/SlateFontDialog/Private/SlateFontDlgWindow.cpp
--- a/Engine/Source/Developer/SlateFontDialog/Private/SlateFontDlgWindow.cpp
+++ b/Engine/Source/Developer/SlateFontDialog/Private/SlateFontDlgWindow.cpp
@@ -11,8 +11,6 @@
 #include "Widgets/Colors/SColorPicker.h"
 #include "Widgets/Colors/SColorBlock.h"
 
-#include "EditorStyleSet.h"
-
 #include <fontconfig.h>
 
 #define LOCTEXT_NAMESPACE "SlateFontDialogNamespace"
@@ -321,11 +319,11 @@ void FSlateFontDlgWindow::OpenFontWindow(FString& OutFontName, float& OutHeight,
 					{
 						if (NewState == ECheckBoxState::Checked)
 						{
-							SampleTextStyle.SetStrikeBrush(*FEditorStyle::GetBrush("DefaultTextUnderline"));
+							SampleTextStyle.SetStrikeBrush(*FAppStyle::Get().GetBrush("DefaultTextUnderline"));
 						}
 						else
 						{
-							SampleTextStyle.SetStrikeBrush(*FEditorStyle::GetBrush("NoBrush"));
+							SampleTextStyle.SetStrikeBrush(*FAppStyle::Get().GetBrush("NoBrush"));
 						}
 
 						SampleTextBlock->SetTextStyle(&SampleTextStyle);
@@ -357,12 +355,12 @@ void FSlateFontDlgWindow::OpenFontWindow(FString& OutFontName, float& OutHeight,
 					{
 						if (NewState == ECheckBoxState::Checked)
 						{
-							SampleTextStyle.SetUnderlineBrush(*FEditorStyle::GetBrush("DefaultTextUnderline"));
+							SampleTextStyle.SetUnderlineBrush(*FAppStyle::Get().GetBrush("DefaultTextUnderline"));
 							EnumAddFlags(OutFlags, EFontImportFlags::EnableUnderline);
 						}
 						else
 						{
-							SampleTextStyle.SetUnderlineBrush(*FEditorStyle::GetBrush("NoBrush"));
+							SampleTextStyle.SetUnderlineBrush(*FAppStyle::Get().GetBrush("NoBrush"));
 							EnumRemoveFlags(OutFlags, EFontImportFlags::EnableUnderline);
 						}
 
@@ -399,7 +397,7 @@ void FSlateFontDlgWindow::OpenFontWindow(FString& OutFontName, float& OutHeight,
 			.HAlign(HAlign_Left)
 			[
 				SNew(SButton)
-				.ButtonStyle(FEditorStyle::Get(), "Menu.Button")
+				.ButtonStyle(FAppStyle::Get(), "Menu.Button")
 				.OnClicked_Raw(this, &FSlateFontDlgWindow::OpenColorPicker)
 				[
 					SNew(SOverlay)
diff --git a/Engine/Source/Developer/SlateFontDialog/SlateFontDialog.Build.cs b/Engine/Source/Developer/SlateFontDialog/SlateFontDialog.Build.cs
--- a/Engine/Source/Developer/SlateFontDialog/SlateFontDialog.Build.cs
+++ b/Engine/Source/Developer/SlateFontDialog/SlateFontDialog.Build.cs
@@ -12,7 +12,6 @@ public class SlateFontDialog : ModuleRules
 				"Core",
 				"CoreUObject",
 				"AppFramework",
-				"EditorStyle",
 				"InputCore",
 				"Slate",
 				"SlateCore"
